Remove references to UE4 in MayaVelocityGridExporter

diff --git a/Engine/Extras/MayaVelocityGridExporter/ReadMe.txt b/Engine/Extras/MayaVelocityGridExporter/ReadMe.txt
--- a/Engine/Extras/MayaVelocityGridExporter/ReadMe.txt
+++ b/Engine/Extras/MayaVelocityGridExporter/ReadMe.txt
@@ -1,4 +1,4 @@
-UE4 Maya Velocity Grid Exporter
+UE Maya Velocity Grid Exporter
 
 Tested on Maya 2012x64,2013x64 & 2014x64 on Windows 7 64bit
 
@@ -11,6 +11,6 @@ For Maya 2013x64 this is located in: [REDACTED_USER_HOME]\Documents\maya\2014-x64
 Now copy the content of the 'shelves' directory to the location of your Maya user preferences 'shelves' directory/
 For Maya 2013x64 this is located in: [REDACTED_USER_HOME]\Documents\maya\2014-x64\prefs\shelves\
 
-Launch Maya. You should now have a shelf tab named "UE4velocityGridExporter" with a blue icon in the shelf. Clicking this icon will launch the
+Launch Maya. You should now have a shelf tab named "UEvelocityGridExporter" with a blue icon in the shelf. Clicking this icon will launch the
 .fga exporter script.
 
